Updated copyright notices

diff --git a/Engine/Source/Programs/RoboMerge/public/js/acknowledge.js b/Engine/Source/Programs/RoboMerge/public/js/acknowledge.js
--- a/Engine/Source/Programs/RoboMerge/public/js/acknowledge.js
+++ b/Engine/Source/Programs/RoboMerge/public/js/acknowledge.js
@@ -1,3 +1,4 @@
+// Copyright 1998-2019 Epic Games, Inc. All Rights Reserved.
 function acknowledge() {
     let queryParams = processQueryParameters(['bot', 'branch', 'cl'])
     if (!queryParams) return;
diff --git a/Engine/Source/Programs/RoboMerge/public/js/boilerplate.js b/Engine/Source/Programs/RoboMerge/public/js/boilerplate.js
--- a/Engine/Source/Programs/RoboMerge/public/js/boilerplate.js
+++ b/Engine/Source/Programs/RoboMerge/public/js/boilerplate.js
@@ -1,3 +1,4 @@
+// Copyright 1998-2019 Epic Games, Inc. All Rights Reserved.
 // boilerplate.js -- This script should be loaded for every Robomerge page to provide a consistent UI and dataset for each page,
 //                   as well as provide common, shared functions to our client-side Javascript.
 
diff --git a/Engine/Source/Programs/RoboMerge/public/js/flow.js b/Engine/Source/Programs/RoboMerge/public/js/flow.js
--- a/Engine/Source/Programs/RoboMerge/public/js/flow.js
+++ b/Engine/Source/Programs/RoboMerge/public/js/flow.js
@@ -1,3 +1,4 @@
+// Copyright 1998-2019 Epic Games, Inc. All Rights Reserved.
 "use strict"
 
 function showFlowGraph(data, botName) {
diff --git a/Engine/Source/Programs/RoboMerge/public/js/graphviz-worker.js b/Engine/Source/Programs/RoboMerge/public/js/graphviz-worker.js
--- a/Engine/Source/Programs/RoboMerge/public/js/graphviz-worker.js
+++ b/Engine/Source/Programs/RoboMerge/public/js/graphviz-worker.js
@@ -1,4 +1,4 @@
-// Copyright 1998-2016 Epic Games, Inc. All Rights Reserved.
+// Copyright 1998-2019 Epic Games, Inc. All Rights Reserved.
 
 importScripts('/js/graphviz-asm.js');
 importScripts('/js/graphviz.js');
diff --git a/Engine/Source/Programs/RoboMerge/public/js/login.js b/Engine/Source/Programs/RoboMerge/public/js/login.js
--- a/Engine/Source/Programs/RoboMerge/public/js/login.js
+++ b/Engine/Source/Programs/RoboMerge/public/js/login.js
@@ -1,4 +1,4 @@
-// Copyright 1998-2018 Epic Games, Inc. All Rights Reserved.
+// Copyright 1998-2019 Epic Games, Inc. All Rights Reserved.
 "use strict"
 
 const COOKIE_REGEX = /\s*(.*?)=(.*)/;
diff --git a/Engine/Source/Programs/RoboMerge/public/js/main.js b/Engine/Source/Programs/RoboMerge/public/js/main.js
--- a/Engine/Source/Programs/RoboMerge/public/js/main.js
+++ b/Engine/Source/Programs/RoboMerge/public/js/main.js
@@ -1,4 +1,4 @@
-// Copyright 1998-2018 Epic Games, Inc. All Rights Reserved.
+// Copyright 1998-2019 Epic Games, Inc. All Rights Reserved.
 "use strict"
 
 function promptFor(map) {
diff --git a/Engine/Source/Programs/RoboMerge/public/js/skip.js b/Engine/Source/Programs/RoboMerge/public/js/skip.js
--- a/Engine/Source/Programs/RoboMerge/public/js/skip.js
+++ b/Engine/Source/Programs/RoboMerge/public/js/skip.js
@@ -1,3 +1,4 @@
+// Copyright 1998-2019 Epic Games, Inc. All Rights Reserved.
 function skipVerify() {
     addOnBeforeUnload('Leaving this page will abort the skip action')
     let queryParams = processQueryParameters(['bot', 'branch', 'cl'], ["reason"])
